nDisplay: Fixed issue where CCWs would spawn on the opposite pole when adding new ones while looking at the top of the stage

diff --git a/Engine/Plugins/Runtime/nDisplayModularFeatures/Source/DisplayClusterLightCardExtender/Private/StageActor/IDisplayClusterStageActor.cpp b/Engine/Plugins/Runtime/nDisplayModularFeatures/Source/DisplayClusterLightCardExtender/Private/StageActor/IDisplayClusterStageActor.cpp
--- a/Engine/Plugins/Runtime/nDisplayModularFeatures/Source/DisplayClusterLightCardExtender/Private/StageActor/IDisplayClusterStageActor.cpp
+++ b/Engine/Plugins/Runtime/nDisplayModularFeatures/Source/DisplayClusterLightCardExtender/Private/StageActor/IDisplayClusterStageActor.cpp
@@ -165,19 +165,19 @@ void IDisplayClusterStageActor::ClampLatitudeAndLongitude(double& InOutLatitude,
 
 FTransform IDisplayClusterStageActor::PositionalParamsToActorTransform(const FDisplayClusterPositionalParams& InParams, const FTransform& InOrigin)
 {
-	const double LocalLatitude = FMath::DegreesToRadians(180.f - InParams.Latitude);
-	const double LocalLongitude = FMath::DegreesToRadians(InParams.Longitude);
+	const double Inclination = FMath::DegreesToRadians(90.0 - InParams.Latitude);
+	const double Azimuth = FMath::DegreesToRadians(InParams.Longitude + 180.0);
 	const double Radius = InParams.DistanceFromCenter + InParams.RadialOffset;
 	
 	FVector Location;
 	{
-		const double CosLat = FMath::Cos(LocalLatitude);
-		const double SinLat = FMath::Sin(LocalLatitude);
-		const double CosLong = FMath::Cos(LocalLongitude);
-		const double SinLong = FMath::Sin(LocalLongitude);
-		Location.X = Radius * CosLat * CosLong;
-		Location.Y = Radius * CosLat * SinLong;
-		Location.Z = Radius * SinLat;
+		const double CosInc = FMath::Cos(Inclination);
+		const double SinInc = FMath::Sin(Inclination);
+		const double CosAz = FMath::Cos(Azimuth);
+		const double SinAz = FMath::Sin(Azimuth);
+		Location.X = Radius * SinInc * CosAz;
+		Location.Y = Radius * SinInc * SinAz;
+		Location.Z = Radius * CosInc;
 	}
 	
 	const FRotator SphericalRotation = FRotator(-InParams.Latitude, InParams.Longitude, 0.0);
@@ -205,8 +205,8 @@ FDisplayClusterPositionalParams IDisplayClusterStageActor::TransformToPositional
 	const double Radius = FVector::Dist(InTransform.GetLocation(), InOrigin.GetLocation());
 	const double DistanceFromCenter = Radius - InRadialOffset;
 	
-	double Latitude = FMath::RadiansToDegrees(FMath::Asin(Location.Z / Radius)) + 180.f;
-	double Longitude = FMath::RadiansToDegrees(FMath::Atan2(Location.Y, Location.X));
+	double Latitude = 90.0 - FMath::RadiansToDegrees(FMath::Acos(Location.Z / Radius));
+	double Longitude = FMath::RadiansToDegrees(FMath::Atan2(Location.Y, Location.X)) - 180.0;
 
 	// Clamp is necessary to maintain correct rotation
 	ClampLatitudeAndLongitude(Latitude, Longitude);
