nDisplay: Fix gizmos not updating on the CCW level instance when moving via orientation values or by dragging in the ICVFX panel.

diff --git a/Engine/Plugins/Runtime/nDisplayModularFeatures/Source/DisplayClusterLightCardExtender/DisplayClusterLightCardExtender.build.cs b/Engine/Plugins/Runtime/nDisplayModularFeatures/Source/DisplayClusterLightCardExtender/DisplayClusterLightCardExtender.build.cs
--- a/Engine/Plugins/Runtime/nDisplayModularFeatures/Source/DisplayClusterLightCardExtender/DisplayClusterLightCardExtender.build.cs
+++ b/Engine/Plugins/Runtime/nDisplayModularFeatures/Source/DisplayClusterLightCardExtender/DisplayClusterLightCardExtender.build.cs
@@ -13,5 +13,10 @@ public class DisplayClusterLightCardExtender : ModuleRules
 				"CoreUObject",
 				"Engine"
 			});
+
+		if (Target.Type == TargetType.Editor)
+		{
+			PrivateDependencyModuleNames.Add("UnrealEd");
+		}
 	}
 }
diff --git a/Engine/Plugins/Runtime/nDisplayModularFeatures/Source/DisplayClusterLightCardExtender/Private/StageActor/IDisplayClusterStageActor.cpp b/Engine/Plugins/Runtime/nDisplayModularFeatures/Source/DisplayClusterLightCardExtender/Private/StageActor/IDisplayClusterStageActor.cpp
--- a/Engine/Plugins/Runtime/nDisplayModularFeatures/Source/DisplayClusterLightCardExtender/Private/StageActor/IDisplayClusterStageActor.cpp
+++ b/Engine/Plugins/Runtime/nDisplayModularFeatures/Source/DisplayClusterLightCardExtender/Private/StageActor/IDisplayClusterStageActor.cpp
@@ -4,6 +4,10 @@
 
 #include "GameFramework/Actor.h"
 
+#if WITH_EDITOR
+#include "Editor.h"
+#endif
+
 const FRotator IDisplayClusterStageActor::PlaneMeshRotation = FRotator(0.0f, -90.0f, 90.0f);
 
 struct FStageActorHelper
@@ -66,6 +70,10 @@ void IDisplayClusterStageActor::UpdateStageActorTransform()
 	if (AActor* Actor = FStageActorHelper::GetActor(this))
 	{
 		Actor->SetActorTransform(Transform);
+
+#if WITH_EDITOR
+		UpdateEditorGizmos();
+#endif
 	}
 }
 
@@ -129,6 +137,10 @@ void IDisplayClusterStageActor::UpdatePositionalParamsFromTransform()
 		Actor->GetTransform(), GetOrigin(), GetRadialOffset());
 
 		SetPositionalParams(PositionalParams);
+
+#if WITH_EDITOR
+		UpdateEditorGizmos();
+#endif
 	}
 }
 
@@ -230,3 +242,19 @@ FDisplayClusterPositionalParams IDisplayClusterStageActor::TransformToPositional
 	
 	return PositionalParams;
 }
+
+#if WITH_EDITOR
+void IDisplayClusterStageActor::UpdateEditorGizmos()
+{
+	if (GEditor)
+	{
+		if (const AActor* Actor = Cast<AActor>(this))
+		{
+			if (Actor->IsSelectedInEditor())
+			{
+				GEditor->NoteSelectionChange(/*bNotify*/ false);
+			}
+		}
+	}
+}
+#endif
diff --git a/Engine/Plugins/Runtime/nDisplayModularFeatures/Source/DisplayClusterLightCardExtender/Public/StageActor/IDisplayClusterStageActor.h b/Engine/Plugins/Runtime/nDisplayModularFeatures/Source/DisplayClusterLightCardExtender/Public/StageActor/IDisplayClusterStageActor.h
--- a/Engine/Plugins/Runtime/nDisplayModularFeatures/Source/DisplayClusterLightCardExtender/Public/StageActor/IDisplayClusterStageActor.h
+++ b/Engine/Plugins/Runtime/nDisplayModularFeatures/Source/DisplayClusterLightCardExtender/Public/StageActor/IDisplayClusterStageActor.h
@@ -96,4 +96,10 @@ public:
 
 	/** Convert an actor transform to positional params */
 	static FDisplayClusterPositionalParams TransformToPositionalParams(const FTransform& InTransform, const FTransform& InOrigin, double InRadialOffset);
+
+#if WITH_EDITOR
+protected:
+	/** Update global selection gizmos */
+	void UpdateEditorGizmos();
+#endif
 };
