Fix for unitialized scale property

diff --git a/Engine/Plugins/Runtime/nDisplayModularFeatures/Source/DisplayClusterLightCardExtender/Public/StageActor/DisplayClusterPositionalParams.h b/Engine/Plugins/Runtime/nDisplayModularFeatures/Source/DisplayClusterLightCardExtender/Public/StageActor/DisplayClusterPositionalParams.h
--- a/Engine/Plugins/Runtime/nDisplayModularFeatures/Source/DisplayClusterLightCardExtender/Public/StageActor/DisplayClusterPositionalParams.h
+++ b/Engine/Plugins/Runtime/nDisplayModularFeatures/Source/DisplayClusterLightCardExtender/Public/StageActor/DisplayClusterPositionalParams.h
@@ -38,5 +38,5 @@ struct FDisplayClusterPositionalParams
 	double RadialOffset = 0.f;
 
 	UPROPERTY(EditAnywhere, BlueprintReadWrite, Category = Orientation)
-	FVector2D Scale;
+	FVector2D Scale = FVector2D(1.0f);
 };
