nDisplay: Fix ICVFX editor scale gizmo not working on CCWs. This was caused by a recent change to CCWs where the mesh components no longer have collision.

diff --git a/Engine/Plugins/Runtime/nDisplayModularFeatures/Source/DisplayClusterLightCardExtender/Private/StageActor/IDisplayClusterStageActor.cpp b/Engine/Plugins/Runtime/nDisplayModularFeatures/Source/DisplayClusterLightCardExtender/Private/StageActor/IDisplayClusterStageActor.cpp
--- a/Engine/Plugins/Runtime/nDisplayModularFeatures/Source/DisplayClusterLightCardExtender/Private/StageActor/IDisplayClusterStageActor.cpp
+++ b/Engine/Plugins/Runtime/nDisplayModularFeatures/Source/DisplayClusterLightCardExtender/Private/StageActor/IDisplayClusterStageActor.cpp
@@ -54,7 +54,7 @@ FBox IDisplayClusterStageActor::GetBoxBounds(bool bLocalSpace) const
 {
 	if (const AActor* Actor = FStageActorHelper::GetActorConst(this))
 	{
-		return Actor->GetComponentsBoundingBox();
+		return Actor->GetComponentsBoundingBox(/* bNonColliding */ true);
 	}
 	return FBox();
 }
