[nDisplay] Fixed CIS validation issue where DisplayClusterModularFeaturesEditor artifacts have paths of 200+ symbols

diff --git a/Engine/Plugins/Runtime/nDisplayModularFeatures/Source/DisplayClusterModularFeaturesEditor/DisplayClusterModularFeaturesEditor.build.cs b/Engine/Plugins/Runtime/nDisplayModularFeatures/Source/DisplayClusterModularFeaturesEditor/DisplayClusterModularFeaturesEditor.build.cs
--- a/Engine/Plugins/Runtime/nDisplayModularFeatures/Source/DisplayClusterModularFeaturesEditor/DisplayClusterModularFeaturesEditor.build.cs
+++ b/Engine/Plugins/Runtime/nDisplayModularFeatures/Source/DisplayClusterModularFeaturesEditor/DisplayClusterModularFeaturesEditor.build.cs
@@ -7,6 +7,8 @@ public class DisplayClusterModularFeaturesEditor : ModuleRules
 {
 	public DisplayClusterModularFeaturesEditor(ReadOnlyTargetRules ROTargetRules) : base(ROTargetRules)
 	{
+		ShortName = "NDCModFeaturesEd";
+
 		PublicDependencyModuleNames.AddRange(
 			new string[] {
 				"Core",
