[nDisplay] In-Camera VFX panel makes level dirty

- Fixed operator== and operator!= of FDisplayClusterPositionalParams
- Added gentle comparison for double member variables

#rnx

diff --git a/Engine/Plugins/Runtime/nDisplayModularFeatures/Source/DisplayClusterLightCardExtender/Public/StageActor/DisplayClusterPositionalParams.h b/Engine/Plugins/Runtime/nDisplayModularFeatures/Source/DisplayClusterLightCardExtender/Public/StageActor/DisplayClusterPositionalParams.h
--- a/Engine/Plugins/Runtime/nDisplayModularFeatures/Source/DisplayClusterLightCardExtender/Public/StageActor/DisplayClusterPositionalParams.h
+++ b/Engine/Plugins/Runtime/nDisplayModularFeatures/Source/DisplayClusterLightCardExtender/Public/StageActor/DisplayClusterPositionalParams.h
@@ -40,20 +40,21 @@ struct FDisplayClusterPositionalParams
 	UPROPERTY(EditAnywhere, BlueprintReadWrite, Category = Orientation)
 	FVector2D Scale = FVector2D(1.0f);
 	
-	FORCEINLINE friend bool operator==(const FDisplayClusterPositionalParams& A, const FDisplayClusterPositionalParams& B)
+	bool operator==(const FDisplayClusterPositionalParams& Other) const
 	{
-		return A.DistanceFromCenter == B.DistanceFromCenter &&
-			A.Longitude == B.Longitude &&
-			A.Latitude == B.Latitude &&
-			A.Spin == B.Spin &&
-			A.Pitch == B.Pitch &&
-			A.Yaw == B.Yaw &&
-			A.RadialOffset == B.RadialOffset &&
-			A.Scale == B.Scale;
+		return
+			FMath::IsNearlyEqual(this->DistanceFromCenter, Other.DistanceFromCenter, KINDA_SMALL_NUMBER) &&
+			FMath::IsNearlyEqual(this->Longitude, Other.Longitude, KINDA_SMALL_NUMBER) &&
+			FMath::IsNearlyEqual(this->Latitude,  Other.Latitude, KINDA_SMALL_NUMBER) &&
+			FMath::IsNearlyEqual(this->Spin,  Other.Spin, KINDA_SMALL_NUMBER) &&
+			FMath::IsNearlyEqual(this->Pitch, Other.Pitch, KINDA_SMALL_NUMBER) &&
+			FMath::IsNearlyEqual(this->Yaw,   Other.Yaw, KINDA_SMALL_NUMBER) &&
+			FMath::IsNearlyEqual(this->RadialOffset, Other.RadialOffset, KINDA_SMALL_NUMBER) &&
+			this->Scale == Other.Scale;
 	}
 
-	FORCEINLINE friend bool operator!=(const FDisplayClusterPositionalParams& A, const FDisplayClusterPositionalParams& B)
+	bool operator!=(const FDisplayClusterPositionalParams& Other) const
 	{
-		return !(A == B);
+		return !(*this == Other);
 	}
 };
